Mover: Adding UE5 Mover-Mass translators along with a Mover Integrations plugin to house various integrations between Mover and other systems

diff --git a/Engine/Plugins/Experimental/MoverIntegrations/MoverIntegrations.uplugin b/Engine/Plugins/Experimental/MoverIntegrations/MoverIntegrations.uplugin
new file mode 100644
--- /dev/null
+++ b/Engine/Plugins/Experimental/MoverIntegrations/MoverIntegrations.uplugin
@@ -0,0 +1,38 @@
+﻿{
+	"FileVersion": 3,
+	"Version": 1,
+	"VersionName": "1.0",
+	"FriendlyName": "Mover Integrations",
+	"Description": "Mover Integrations is a Unreal Engine plugin acting as an umbrella to cover a variety of modules supporting Mover's integration with other plugins, such as animation, AI, and other gameplay systems.",
+	"Category": "Other",
+	"CreatedBy": "Epic Games, Inc.",
+	"CreatedByURL": "https://epicgames.com",
+	"DocsURL": "",
+	"MarketplaceURL": "",
+	"CanContainContent": true,
+	"IsBetaVersion": false,
+	"IsExperimentalVersion": true,
+	"Installed": false,
+	"Modules": [
+		{
+			"Name": "MoverIntegrations",
+			"Type": "Runtime",
+			"LoadingPhase": "Default"
+		},
+		{
+			"Name": "MoverMassIntegration",
+			"Type": "Runtime",
+			"LoadingPhase": "Default"
+		}
+	],
+	"Plugins": [
+		{
+			"Name": "Mover",
+			"Enabled": true
+		},
+		{
+			"Name": "MassGameplay",
+			"Enabled": true
+		}
+	]
+}
\ No newline at end of file
diff --git a/Engine/Plugins/Experimental/MoverIntegrations/Source/MoverIntegrations/MoverIntegrations.Build.cs b/Engine/Plugins/Experimental/MoverIntegrations/Source/MoverIntegrations/MoverIntegrations.Build.cs
new file mode 100644
--- /dev/null
+++ b/Engine/Plugins/Experimental/MoverIntegrations/Source/MoverIntegrations/MoverIntegrations.Build.cs
@@ -0,0 +1,51 @@
+﻿// Copyright Epic Games, Inc. All Rights Reserved.
+
+using UnrealBuildTool;
+
+public class MoverIntegrations : ModuleRules
+{
+	public MoverIntegrations(ReadOnlyTargetRules Target) : base(Target)
+	{
+		PCHUsage = ModuleRules.PCHUsageMode.UseExplicitOrSharedPCHs;
+		
+		PublicIncludePaths.AddRange(
+			new string[] {
+				// ... add public include paths required here ...
+			}
+			);
+				
+		
+		PrivateIncludePaths.AddRange(
+			new string[] {
+				// ... add other private include paths required here ...
+			}
+			);
+			
+		
+		PublicDependencyModuleNames.AddRange(
+			new string[]
+			{
+				"Core",
+				// ... add other public dependencies that you statically link with here ...
+			}
+			);
+			
+		
+		PrivateDependencyModuleNames.AddRange(
+			new string[]
+			{
+				"CoreUObject",
+				"Engine",
+				// ... add private dependencies that you statically link with here ...	
+			}
+			);
+		
+		
+		DynamicallyLoadedModuleNames.AddRange(
+			new string[]
+			{
+				// ... add any modules that your module loads dynamically here ...
+			}
+			);
+	}
+}
diff --git a/Engine/Plugins/Experimental/MoverIntegrations/Source/MoverIntegrations/Private/MoverIntegrations.cpp b/Engine/Plugins/Experimental/MoverIntegrations/Source/MoverIntegrations/Private/MoverIntegrations.cpp
new file mode 100644
--- /dev/null
+++ b/Engine/Plugins/Experimental/MoverIntegrations/Source/MoverIntegrations/Private/MoverIntegrations.cpp
@@ -0,0 +1,20 @@
+﻿// Copyright Epic Games, Inc. All Rights Reserved.
+
+#include "MoverIntegrations.h"
+
+#define LOCTEXT_NAMESPACE "FMoverIntegrationsModule"
+
+void FMoverIntegrationsModule::StartupModule()
+{
+	// This code will execute after your module is loaded into memory; the exact timing is specified in the .uplugin file per-module
+}
+
+void FMoverIntegrationsModule::ShutdownModule()
+{
+	// This function may be called during shutdown to clean up your module.  For modules that support dynamic reloading,
+	// we call this function before unloading the module.
+}
+
+#undef LOCTEXT_NAMESPACE
+	
+IMPLEMENT_MODULE(FMoverIntegrationsModule, MoverIntegrations)
\ No newline at end of file
diff --git a/Engine/Plugins/Experimental/MoverIntegrations/Source/MoverIntegrations/Public/MoverIntegrations.h b/Engine/Plugins/Experimental/MoverIntegrations/Source/MoverIntegrations/Public/MoverIntegrations.h
new file mode 100644
--- /dev/null
+++ b/Engine/Plugins/Experimental/MoverIntegrations/Source/MoverIntegrations/Public/MoverIntegrations.h
@@ -0,0 +1,17 @@
+﻿// Copyright Epic Games, Inc. All Rights Reserved.
+
+#pragma once
+
+#include "Modules/ModuleManager.h"
+
+/**
+ * Module that serves as umbrella of a variety of integration modules between Mover and other systems (such as animation, AI, and other gameplay systems)
+ */
+class FMoverIntegrationsModule : public IModuleInterface
+{
+public:
+
+	/** IModuleInterface implementation */
+	virtual void StartupModule() override;
+	virtual void ShutdownModule() override;
+};
diff --git a/Engine/Plugins/Experimental/MoverIntegrations/Source/MoverMassIntegration/MoverMassIntegration.Build.cs b/Engine/Plugins/Experimental/MoverIntegrations/Source/MoverMassIntegration/MoverMassIntegration.Build.cs
new file mode 100644
--- /dev/null
+++ b/Engine/Plugins/Experimental/MoverIntegrations/Source/MoverMassIntegration/MoverMassIntegration.Build.cs
@@ -0,0 +1,32 @@
+﻿// Copyright Epic Games, Inc. All Rights Reserved.
+
+using UnrealBuildTool;
+
+public class MoverMassIntegration : ModuleRules
+{
+    public MoverMassIntegration(ReadOnlyTargetRules Target) : base(Target)
+    {
+        PCHUsage = ModuleRules.PCHUsageMode.UseExplicitOrSharedPCHs;
+
+        PublicDependencyModuleNames.AddRange(
+            new string[]
+            {
+                "Core",
+            }
+        );
+
+        PrivateDependencyModuleNames.AddRange(
+            new string[]
+            {
+                "CoreUObject",
+                "Engine",
+                "Mover",
+                "MassActors",
+                "MassEntity",
+                "MassSpawner",
+                "MassCommon",
+                "MassMovement"
+            }
+        );
+    }
+}
\ No newline at end of file
diff --git a/Engine/Plugins/Experimental/MoverIntegrations/Source/MoverMassIntegration/Private/MoverMassAgentTraits.cpp b/Engine/Plugins/Experimental/MoverIntegrations/Source/MoverMassIntegration/Private/MoverMassAgentTraits.cpp
new file mode 100644
--- /dev/null
+++ b/Engine/Plugins/Experimental/MoverIntegrations/Source/MoverMassIntegration/Private/MoverMassAgentTraits.cpp
@@ -0,0 +1,120 @@
+﻿// Copyright Epic Games, Inc. All Rights Reserved.
+
+
+#include "MoverMassAgentTraits.h"
+#include "DefaultMovementSet/NavMoverComponent.h"
+#include "MassMovementFragments.h"
+#include "MassEntityTemplate.h"
+#include "MassEntityTemplateRegistry.h"
+#include "MoverMassTranslators.h"
+#include "GameFramework/Actor.h"
+#include "VisualLogger/VisualLogger.h"
+#include "MassEntityView.h"
+
+namespace UE::Mass::Mover::Private
+{
+	template<typename T>
+	T* AsComponent(UObject& Owner)
+	{
+		T* Component = nullptr;
+		if (AActor* AsActor = Cast<AActor>(&Owner))
+		{
+			Component = AsActor->FindComponentByClass<T>();
+		}
+		else
+		{
+			Component = Cast<T>(&Owner);
+		}
+
+		UE_CVLOG_UELOG(Component == nullptr, &Owner, LogMass, Error, TEXT("Trying to extract %s from %s failed")
+					   , *T::StaticClass()->GetName(), *Owner.GetName());
+
+		return Component;
+	}
+}
+
+
+//----------------------------------------------------------------------//
+//  UMoverMassAgentTrait
+//----------------------------------------------------------------------//
+void UMoverMassAgentTrait::BuildTemplate(FMassEntityTemplateBuildContext& BuildContext, const UWorld& World) const
+{
+	BuildContext.AddFragment<FNavMoverComponentWrapperFragment>();
+	BuildContext.AddFragment<FMassVelocityFragment>();
+	BuildContext.AddFragment<FTransformFragment>();
+	
+	BuildContext.GetMutableObjectFragmentInitializers().Add([this](UObject& Owner, FMassEntityView& EntityView, const EMassTranslationDirection CurrentDirection)
+	{
+		if (UNavMoverComponent* NavMoverComponent = UE::Mass::Mover::Private::AsComponent<UNavMoverComponent>(Owner))
+		{
+			FNavMoverComponentWrapperFragment& ComponentFragment = EntityView.GetFragmentData<FNavMoverComponentWrapperFragment>();
+			ComponentFragment.Component = NavMoverComponent;
+
+			FMassVelocityFragment& VelocityFragment = EntityView.GetFragmentData<FMassVelocityFragment>();
+
+			// the entity is the authority
+			if (CurrentDirection ==  EMassTranslationDirection::MassToActor)
+			{
+				NavMoverComponent->RequestDirectMove(VelocityFragment.Value, /*bForceMaxSpeed*/false);
+			}
+			// actor is the authority
+			else
+			{
+				VelocityFragment.Value = NavMoverComponent->GetVelocityForNavMovement();
+			}
+
+			if (bSyncTransform)
+			{
+				FTransformFragment& TransformFragment = EntityView.GetFragmentData<FTransformFragment>();
+
+				// the entity is the authority
+				if (CurrentDirection ==  EMassTranslationDirection::MassToActor)
+				{
+					if (USceneComponent* UpdatedObjectAsSceneComponent = Cast<USceneComponent>(NavMoverComponent->GetUpdatedObject()))
+					{
+						// TODO: Mover also doesn't like setting transforms directly and may cause a warning about outside systems modifying the updated component
+						UpdatedObjectAsSceneComponent->SetWorldTransform(TransformFragment.GetTransform());
+					}
+				}
+				// actor is the authority
+				else
+				{
+					if (const USceneComponent* UpdatedObjectAsSceneComponent = Cast<USceneComponent>(NavMoverComponent->GetUpdatedObject()))
+					{
+						TransformFragment.SetTransform(UpdatedObjectAsSceneComponent->GetComponentTransform());
+					}
+				}
+			}
+		}
+	});
+
+	if (EnumHasAnyFlags(SyncDirection, EMassTranslationDirection::ActorToMass) || BuildContext.IsInspectingData())
+	{
+		BuildContext.AddTranslator<UMassNavMoverToMassTranslator>();
+	}
+
+	if (EnumHasAnyFlags(SyncDirection, EMassTranslationDirection::MassToActor) || BuildContext.IsInspectingData())
+	{
+		BuildContext.AddTranslator<UMassToNavMoverTranslator>();
+	}
+}
+
+//----------------------------------------------------------------------//
+//  UMoverMassAgentOrientationSyncTrait
+//----------------------------------------------------------------------//
+void UMoverMassAgentOrientationSyncTrait::BuildTemplate(FMassEntityTemplateBuildContext& BuildContext, const UWorld& World) const
+{
+	BuildContext.RequireFragment<FNavMoverComponentWrapperFragment>();
+
+	if (EnumHasAnyFlags(SyncDirection, EMassTranslationDirection::ActorToMass) 
+		|| BuildContext.IsInspectingData())
+	{
+		BuildContext.AddTranslator<UMassNavMoverActorOrientationToMassTranslator>();
+	}
+	
+	if (EnumHasAnyFlags(SyncDirection, EMassTranslationDirection::MassToActor)
+		|| BuildContext.IsInspectingData())
+	{
+		BuildContext.AddTranslator<UMassOrientationToNavMoverActorOrientationTranslator>();
+	}
+}
\ No newline at end of file
diff --git a/Engine/Plugins/Experimental/MoverIntegrations/Source/MoverMassIntegration/Private/MoverMassIntegration.cpp b/Engine/Plugins/Experimental/MoverIntegrations/Source/MoverMassIntegration/Private/MoverMassIntegration.cpp
new file mode 100644
--- /dev/null
+++ b/Engine/Plugins/Experimental/MoverIntegrations/Source/MoverMassIntegration/Private/MoverMassIntegration.cpp
@@ -0,0 +1,19 @@
+﻿// Copyright Epic Games, Inc. All Rights Reserved.
+
+#include "MoverMassIntegration.h"
+
+#define LOCTEXT_NAMESPACE "FMoverMassIntegrationModule"
+
+void FMoverMassIntegrationModule::StartupModule()
+{
+    
+}
+
+void FMoverMassIntegrationModule::ShutdownModule()
+{
+    
+}
+
+#undef LOCTEXT_NAMESPACE
+    
+IMPLEMENT_MODULE(FMoverMassIntegrationModule, MoverMassIntegration)
\ No newline at end of file
diff --git a/Engine/Plugins/Experimental/MoverIntegrations/Source/MoverMassIntegration/Private/MoverMassTranslators.cpp b/Engine/Plugins/Experimental/MoverIntegrations/Source/MoverMassIntegration/Private/MoverMassTranslators.cpp
new file mode 100644
--- /dev/null
+++ b/Engine/Plugins/Experimental/MoverIntegrations/Source/MoverMassIntegration/Private/MoverMassTranslators.cpp
@@ -0,0 +1,236 @@
+﻿// Copyright Epic Games, Inc. All Rights Reserved.
+
+#include "MoverMassTranslators.h"
+#include "DefaultMovementSet/NavMoverComponent.h"
+#include "VisualLogger/VisualLogger.h"
+#include "MassEntityManager.h"
+#include "MassCommonTypes.h"
+#include "MassExecutionContext.h"
+#include "MassMovementFragments.h"
+#include "MassSpawnerTypes.h"
+#include "MassDebugger.h"
+
+DECLARE_LOG_CATEGORY_EXTERN(LogMassMoverTranslator, Log, All);
+DEFINE_LOG_CATEGORY(LogMassMoverTranslator);
+
+DECLARE_LOG_CATEGORY_EXTERN(LogMassMoverDivergence, Log, All);
+DEFINE_LOG_CATEGORY(LogMassMoverDivergence);
+
+//----------------------------------------------------------------------//
+//  UMassNavMoverToMassTranslator
+//----------------------------------------------------------------------//
+UMassNavMoverToMassTranslator::UMassNavMoverToMassTranslator()
+	: EntityQuery(*this)
+{
+	ExecutionFlags = (int32)EProcessorExecutionFlags::AllNetModes;
+	ExecutionOrder.ExecuteInGroup = UE::Mass::ProcessorGroupNames::SyncWorldToMass;
+	RequiredTags.Add<FMassNavMoverCopyToMassTag>();
+}
+
+void UMassNavMoverToMassTranslator::ConfigureQueries(const TSharedRef<FMassEntityManager>& EntityManager)
+{
+	AddRequiredTagsToQuery(EntityQuery);
+	EntityQuery.AddRequirement<FNavMoverComponentWrapperFragment>(EMassFragmentAccess::ReadOnly);
+	EntityQuery.AddRequirement<FTransformFragment>(EMassFragmentAccess::ReadWrite);
+	EntityQuery.AddRequirement<FMassVelocityFragment>(EMassFragmentAccess::ReadWrite);
+	EntityQuery.AddRequirement<FMassDesiredMovementFragment>(EMassFragmentAccess::ReadWrite);
+}
+
+void UMassNavMoverToMassTranslator::Execute(FMassEntityManager& EntityManager, FMassExecutionContext& Context)
+{
+	EntityQuery.ForEachEntityChunk(Context, [this](FMassExecutionContext& Context)
+	{
+		const TConstArrayView<FNavMoverComponentWrapperFragment> ComponentList = Context.GetFragmentView<FNavMoverComponentWrapperFragment>();
+		const TArrayView<FTransformFragment> LocationList = Context.GetMutableFragmentView<FTransformFragment>();
+		const TArrayView<FMassVelocityFragment> VelocityList = Context.GetMutableFragmentView<FMassVelocityFragment>();
+		const TArrayView<FMassDesiredMovementFragment> DesiredMovementList = Context.GetMutableFragmentView<FMassDesiredMovementFragment>();
+
+		for (FMassExecutionContext::FEntityIterator EntityIt = Context.CreateEntityIterator(); EntityIt; ++EntityIt)
+		{
+			if (const UNavMoverComponent* AsMovementComponent = ComponentList[EntityIt].Component.Get())
+			{
+				LocationList[EntityIt].GetMutableTransform().SetLocation(AsMovementComponent->GetFeetLocation());
+				VelocityList[EntityIt].Value = AsMovementComponent->GetVelocityForNavMovement();
+				
+				DesiredMovementList[EntityIt].DesiredMaxSpeedOverride = AsMovementComponent->GetMaxSpeedForNavMovement();	
+
+#if WITH_MASSENTITY_DEBUG
+				const bool bDisplayDebug = UE::Mass::Debug::IsDebuggingEntity(Context.GetEntity(EntityIt));
+				if (bDisplayDebug)
+				{
+					const FVector& PreviousVelocity = VelocityList[EntityIt].DebugPreviousValue;
+					const FVector ZOffset(0,0,5);
+					const FVector& Location = LocationList[EntityIt].GetTransform().GetLocation() + ZOffset;
+
+					constexpr FVector::FReal VelocityDeltaSquared = 1;
+					if (FVector::DistSquared(PreviousVelocity, AsMovementComponent->GetVelocityForNavMovement()) > VelocityDeltaSquared)
+					{
+						// Draw expected and current velocities
+						UE_VLOG_ARROW(this, LogMassMoverDivergence, Log, Location,
+							Location + VelocityList[EntityIt].Value, FColor::Orange, TEXT("Current\nSpeed %.1f"), VelocityList[EntityIt].Value.Size());
+
+						UE_VLOG_ARROW(this, LogMassMoverDivergence, Log, Location,
+							Location + PreviousVelocity, FColor::Green, TEXT("Expected\nSpeed %.1f"), PreviousVelocity.Size());
+					}
+				}
+#endif //WITH_MASSENTITY_DEBUG
+			}
+		}
+	});
+}
+
+//----------------------------------------------------------------------//
+//  UMassToNavMoverTranslator
+//----------------------------------------------------------------------//
+UMassToNavMoverTranslator::UMassToNavMoverTranslator()
+	: EntityQuery(*this)
+{
+	ExecutionFlags = (int32)EProcessorExecutionFlags::AllNetModes;
+	ExecutionOrder.ExecuteInGroup = UE::Mass::ProcessorGroupNames::UpdateWorldFromMass;
+	ExecutionOrder.ExecuteAfter.Add(UE::Mass::ProcessorGroupNames::Movement);
+	RequiredTags.Add<FMassCopyToNavMoverTag>();
+}
+
+void UMassToNavMoverTranslator::ConfigureQueries(const TSharedRef<FMassEntityManager>& EntityManager)
+{
+	AddRequiredTagsToQuery(EntityQuery);
+	EntityQuery.AddRequirement<FNavMoverComponentWrapperFragment>(EMassFragmentAccess::ReadWrite);
+
+#if WITH_MASSENTITY_DEBUG	
+	EntityQuery.AddRequirement<FMassVelocityFragment>(EMassFragmentAccess::ReadWrite);
+#else
+	EntityQuery.AddRequirement<FMassVelocityFragment>(EMassFragmentAccess::ReadOnly);
+#endif //WITH_MASSENTITY_DEBUG
+	
+	EntityQuery.AddRequirement<FTransformFragment>(EMassFragmentAccess::ReadOnly);
+}
+
+void UMassToNavMoverTranslator::Execute(FMassEntityManager& EntityManager, FMassExecutionContext& Context)
+{
+	EntityQuery.ForEachEntityChunk(Context, [this](FMassExecutionContext& Context)
+	{
+		const TArrayView<FNavMoverComponentWrapperFragment> ComponentList = Context.GetMutableFragmentView<FNavMoverComponentWrapperFragment>();
+
+#if WITH_MASSGAMEPLAY_DEBUG
+		const TArrayView<FMassVelocityFragment> VelocityList = Context.GetMutableFragmentView<FMassVelocityFragment>();
+#else
+		const TConstArrayView<FMassVelocityFragment> VelocityList = Context.GetFragmentView<FMassVelocityFragment>();
+#endif //WITH_MASSGAMEPLAY_DEBUG
+
+		const TConstArrayView<FTransformFragment> LocationList = Context.GetFragmentView<FTransformFragment>();
+		
+		for (FMassExecutionContext::FEntityIterator EntityIt = Context.CreateEntityIterator(); EntityIt; ++EntityIt)
+		{
+			if (UNavMoverComponent* AsMovementComponent = ComponentList[EntityIt].Component.Get())
+			{
+				FVector RequestedMove = VelocityList[EntityIt].Value;
+				RequestedMove.Z = 0.0f;
+
+#if WITH_MASSGAMEPLAY_DEBUG
+				// Store requested velocity.
+				VelocityList[EntityIt].DebugPreviousValue = RequestedMove;
+#endif //WITH_MASSGAMEPLAY_DEBUG
+			
+				AsMovementComponent->RequestDirectMove(RequestedMove, /*bForceMaxSpeed=*/false);
+				
+#if WITH_MASSGAMEPLAY_DEBUG
+				bool bDisplayDebug = UE::Mass::Debug::IsDebuggingEntity(Context.GetEntity(EntityIt));
+				if (bDisplayDebug)
+				{
+					const FVector ActorLocation = AsMovementComponent->GetFeetLocation();
+					const FVector EntityLocation = LocationList[EntityIt].GetTransform().GetLocation();
+					
+					UE_VLOG_ARROW(AsMovementComponent, LogMassMoverTranslator, Display, ActorLocation, ActorLocation + RequestedMove, FColor::Green, TEXT("Requested Move: %s"), *(RequestedMove.ToString()));
+					UE_VLOG_SPHERE(AsMovementComponent, LogMassMoverTranslator, Display, EntityLocation, 5.0f, FColor::White, TEXT("EntityLocation: %s"), *(RequestedMove.ToString()));
+					UE_VLOG_SPHERE(AsMovementComponent, LogMassMoverTranslator, Display, EntityLocation + RequestedMove, 5.0f, FColor::Blue, TEXT("EntityPrediction: %s"), *(RequestedMove.ToString()));
+				}
+#endif // WITH_MASSGAMEPLAY_DEBUG
+
+			}
+		}
+	});
+}
+
+
+//----------------------------------------------------------------------//
+//  UMassNavMoverActorOrientationToMassTranslator
+//----------------------------------------------------------------------//
+UMassNavMoverActorOrientationToMassTranslator::UMassNavMoverActorOrientationToMassTranslator()
+: EntityQuery(*this)
+{
+	ExecutionFlags = (int32)EProcessorExecutionFlags::AllNetModes;
+	ExecutionOrder.ExecuteInGroup = UE::Mass::ProcessorGroupNames::SyncWorldToMass;
+	RequiredTags.Add<FMassNavMoverActorOrientationCopyToMassTag>();
+	bRequiresGameThreadExecution = true;
+}
+
+void UMassNavMoverActorOrientationToMassTranslator::ConfigureQueries(const TSharedRef<FMassEntityManager>& EntityManager)
+{
+	AddRequiredTagsToQuery(EntityQuery);
+	EntityQuery.AddRequirement<FNavMoverComponentWrapperFragment>(EMassFragmentAccess::ReadOnly);
+	EntityQuery.AddRequirement<FTransformFragment>(EMassFragmentAccess::ReadWrite);
+}
+
+void UMassNavMoverActorOrientationToMassTranslator::Execute(FMassEntityManager& EntityManager, FMassExecutionContext& Context)
+{
+	EntityQuery.ForEachEntityChunk(Context, [this](FMassExecutionContext& Context)
+	{
+		const TConstArrayView<FNavMoverComponentWrapperFragment> ComponentList = Context.GetFragmentView<FNavMoverComponentWrapperFragment>();
+		const TArrayView<FTransformFragment> LocationList = Context.GetMutableFragmentView<FTransformFragment>();
+		
+		for (FMassExecutionContext::FEntityIterator EntityIt = Context.CreateEntityIterator(); EntityIt; ++EntityIt)
+		{
+			if (const UNavMoverComponent* AsNavMoverComponent = ComponentList[EntityIt].Component.Get())
+			{
+				if (const USceneComponent* UpdatedComponent = Cast<USceneComponent>(AsNavMoverComponent->GetUpdatedObject()))
+				{
+					LocationList[EntityIt].GetMutableTransform().SetRotation(UpdatedComponent->GetComponentTransform().GetRotation());
+				}
+			}
+		}
+	});
+}
+
+//----------------------------------------------------------------------//
+//  UMassOrientationToNavMoverActorOrientationTranslator
+//----------------------------------------------------------------------//
+UMassOrientationToNavMoverActorOrientationTranslator::UMassOrientationToNavMoverActorOrientationTranslator()
+: EntityQuery(*this)
+{
+	ExecutionFlags = (int32)EProcessorExecutionFlags::AllNetModes;
+	ExecutionOrder.ExecuteInGroup = UE::Mass::ProcessorGroupNames::UpdateWorldFromMass;
+	ExecutionOrder.ExecuteAfter.Add(UE::Mass::ProcessorGroupNames::Movement);
+	RequiredTags.Add<FMassOrientationCopyToNavMoverActorOrientationTag>();
+	bRequiresGameThreadExecution = true;
+}
+
+void UMassOrientationToNavMoverActorOrientationTranslator::ConfigureQueries(const TSharedRef<FMassEntityManager>& EntityManager)
+{
+	AddRequiredTagsToQuery(EntityQuery);
+	EntityQuery.AddRequirement<FNavMoverComponentWrapperFragment>(EMassFragmentAccess::ReadWrite);
+	EntityQuery.AddRequirement<FTransformFragment>(EMassFragmentAccess::ReadOnly);
+	EntityQuery.RequireMutatingWorldAccess(); // due to mutating World by setting component rotation
+}
+
+void UMassOrientationToNavMoverActorOrientationTranslator::Execute(FMassEntityManager& EntityManager, FMassExecutionContext& Context)
+{
+	EntityQuery.ForEachEntityChunk(Context, [this](FMassExecutionContext& Context)
+	{
+		const TArrayView<FNavMoverComponentWrapperFragment> ComponentList = Context.GetMutableFragmentView<FNavMoverComponentWrapperFragment>();
+		const TConstArrayView<FTransformFragment> TransformList = Context.GetFragmentView<FTransformFragment>();
+		
+		for (FMassExecutionContext::FEntityIterator EntityIt = Context.CreateEntityIterator(); EntityIt; ++EntityIt)
+		{
+			if (const UNavMoverComponent* AsNavMoverComponent = ComponentList[EntityIt].Component.Get())
+			{
+				if (USceneComponent* UpdatedComponent = Cast<USceneComponent>(AsNavMoverComponent->GetUpdatedObject()))
+				{
+					// TODO: Set OrientToMovement to true or false here - currently this isn't an option the Mover component but it should be
+					// TODO: Mover also doesn't like setting rotation directly and may cause a warning about outside systems modifying the updated component
+					const FTransformFragment& Transform = TransformList[EntityIt];
+					UpdatedComponent->SetWorldRotation(Transform.GetTransform().GetRotation());
+				}
+			}
+		}
+	});
+}
diff --git a/Engine/Plugins/Experimental/MoverIntegrations/Source/MoverMassIntegration/Public/MoverMassAgentTraits.h b/Engine/Plugins/Experimental/MoverIntegrations/Source/MoverMassIntegration/Public/MoverMassAgentTraits.h
new file mode 100644
--- /dev/null
+++ b/Engine/Plugins/Experimental/MoverIntegrations/Source/MoverMassIntegration/Public/MoverMassAgentTraits.h
@@ -0,0 +1,45 @@
+﻿// Copyright Epic Games, Inc. All Rights Reserved.
+
+#pragma once
+
+#include "MassActors/Public/MassAgentTraits.h"
+#include "MoverMassAgentTraits.generated.h"
+
+struct FMassEntityTemplateBuildContext;
+
+#define UE_API MOVERMASSINTEGRATION_API
+
+/**
+ * The trait initializes the entity with a NavMoverComponent so Mover and Mass can communicate movement intent and velocity.
+ * This trait also sets up necessary translators for these systems based off of Mass sync direction.
+ * Note: This trait requires a NavMoverComponent and a MoverComponent to work properly
+ */
+UCLASS(MinimalAPI, BlueprintType, EditInlineNew, CollapseCategories, meta = (DisplayName = "Agent Mover Sync"))
+class UMoverMassAgentTrait : public UMassAgentSyncTrait
+{
+	GENERATED_BODY()
+
+protected:
+	UE_API virtual void BuildTemplate(FMassEntityTemplateBuildContext& BuildContext, const UWorld& World) const override;
+
+	/**
+	 * Whether this trait should sync the transform of actor<->entity (based on translation direction)
+	 * TODO: Currently Mover doesn't like outside modification of rotation and may throw a warning if Mover's transform gets set from the entity. It may also cause a rollback.
+	 */
+	UPROPERTY(EditAnywhere, Category = Mass)
+	bool bSyncTransform = false;
+};
+
+/**
+ * This trait sets up required translators for orientation syncing between Mover and Mass based off of Mass sync direction
+ */
+UCLASS(MinimalAPI, BlueprintType, EditInlineNew, CollapseCategories, meta = (DisplayName = "Agent Mover Orientation Sync"))
+class UMoverMassAgentOrientationSyncTrait : public UMassAgentSyncTrait
+{
+	GENERATED_BODY()
+
+protected:
+	UE_API virtual void BuildTemplate(FMassEntityTemplateBuildContext& BuildContext, const UWorld& World) const override;
+};
+
+#undef UE_API
diff --git a/Engine/Plugins/Experimental/MoverIntegrations/Source/MoverMassIntegration/Public/MoverMassIntegration.h b/Engine/Plugins/Experimental/MoverIntegrations/Source/MoverMassIntegration/Public/MoverMassIntegration.h
new file mode 100644
--- /dev/null
+++ b/Engine/Plugins/Experimental/MoverIntegrations/Source/MoverMassIntegration/Public/MoverMassIntegration.h
@@ -0,0 +1,15 @@
+﻿// Copyright Epic Games, Inc. All Rights Reserved.
+
+#pragma once
+
+#include "Modules/ModuleManager.h"
+
+/**
+ * Module for Mover-Mass integrations such as Mover specific Mass translators and Mass traits
+ */
+class FMoverMassIntegrationModule : public IModuleInterface
+{
+public:
+    virtual void StartupModule() override;
+    virtual void ShutdownModule() override;
+};
diff --git a/Engine/Plugins/Experimental/MoverIntegrations/Source/MoverMassIntegration/Public/MoverMassTranslators.h b/Engine/Plugins/Experimental/MoverIntegrations/Source/MoverMassIntegration/Public/MoverMassTranslators.h
new file mode 100644
--- /dev/null
+++ b/Engine/Plugins/Experimental/MoverIntegrations/Source/MoverMassIntegration/Public/MoverMassTranslators.h
@@ -0,0 +1,123 @@
+﻿// Copyright Epic Games, Inc. All Rights Reserved.
+
+#pragma once
+
+#include "MassCommonFragments.h"
+#include "MassTranslator.h"
+#include "MassObserverProcessor.h"
+#include "MoverMassTranslators.generated.h"
+
+class UNavMoverComponent;
+
+#define UE_API MOVERMASSINTEGRATION_API
+
+USTRUCT()
+struct FNavMoverComponentWrapperFragment : public FObjectWrapperFragment
+{
+	GENERATED_BODY()
+	TWeakObjectPtr<UNavMoverComponent> Component;
+};
+
+USTRUCT()
+struct FMassNavMoverCopyToMassTag : public FMassTag
+{
+	GENERATED_BODY()
+};
+
+/**
+ * Translator in charge of going from Mover->Mass
+ * Sets Mass fragments for keeping track of Location (feet location), Velocity and MaxSpeed
+ */
+UCLASS(MinimalAPI)
+class UMassNavMoverToMassTranslator : public UMassTranslator
+{
+	GENERATED_BODY()
+
+public:
+	UE_API UMassNavMoverToMassTranslator();
+
+protected:
+	UE_API virtual void ConfigureQueries(const TSharedRef<FMassEntityManager>& EntityManager) override;
+	UE_API virtual void Execute(FMassEntityManager& EntityManager, FMassExecutionContext& Context) override;
+
+	FMassEntityQuery EntityQuery;
+};
+
+USTRUCT()
+struct FMassCopyToNavMoverTag : public FMassTag
+{
+	GENERATED_BODY()
+};
+
+/**
+ * Translator in charge of going from Mass->Mover
+ * Uses the NavMoverWrapperFragment to get a NavMoverComponent and request movement similarly to the PathFollowingComponent
+ * See @UNavMoverComponent for details of how Mover consumes the move intent
+ */
+UCLASS(MinimalAPI)
+class UMassToNavMoverTranslator : public UMassTranslator
+{
+	GENERATED_BODY()
+
+public:
+	UE_API UMassToNavMoverTranslator();
+
+protected:
+	UE_API virtual void ConfigureQueries(const TSharedRef<FMassEntityManager>& EntityManager) override;
+	UE_API virtual void Execute(FMassEntityManager& EntityManager, FMassExecutionContext& Context) override;
+
+	FMassEntityQuery EntityQuery;
+};
+
+USTRUCT()
+struct FMassNavMoverActorOrientationCopyToMassTag : public FMassTag
+{
+	GENERATED_BODY()
+};
+
+/**
+ * Translator in charge of going from Mover->Mass for orientation
+ * Modifies the transform fragment based off of Mover's rotation
+ */
+UCLASS(MinimalAPI)
+class UMassNavMoverActorOrientationToMassTranslator : public UMassTranslator
+{
+	GENERATED_BODY()
+
+public:
+	UE_API UMassNavMoverActorOrientationToMassTranslator();
+
+protected:
+	UE_API virtual void ConfigureQueries(const TSharedRef<FMassEntityManager>& EntityManager) override;
+	UE_API virtual void Execute(FMassEntityManager& EntityManager, FMassExecutionContext& Context) override;
+
+	FMassEntityQuery EntityQuery;
+};
+
+USTRUCT()
+struct FMassOrientationCopyToNavMoverActorOrientationTag : public FMassTag
+{
+	GENERATED_BODY()
+};
+
+/**
+ * Translator in charge of going from Mass->Mover for orientation
+ * Modifies the UpdatedComponent of the NavMoverComponent based off of the entities transform fragment rotation
+ * TODO: Currently Mover doesn't like outside modification of rotation and may throw a warning. It may also cause a rollback.
+ */
+UCLASS(MinimalAPI)
+class UMassOrientationToNavMoverActorOrientationTranslator : public UMassTranslator
+{
+	GENERATED_BODY()
+
+public:
+	UE_API UMassOrientationToNavMoverActorOrientationTranslator();
+
+protected:
+	UE_API virtual void ConfigureQueries(const TSharedRef<FMassEntityManager>& EntityManager) override;
+	UE_API virtual void Execute(FMassEntityManager& EntityManager, FMassExecutionContext& Context) override;
+
+	FMassEntityQuery EntityQuery;
+};
+
+#undef UE_API
