CIS issue#936427: Compile errors in EngineRuntimeTests.h, MoverMassTranslators.h and MoverMassTranslators.cpp

#rnx

diff --git a/Engine/Plugins/Experimental/MoverIntegrations/Source/MoverMassIntegration/Public/MoverMassTranslators.h b/Engine/Plugins/Experimental/MoverIntegrations/Source/MoverMassIntegration/Public/MoverMassTranslators.h
--- a/Engine/Plugins/Experimental/MoverIntegrations/Source/MoverMassIntegration/Public/MoverMassTranslators.h
+++ b/Engine/Plugins/Experimental/MoverIntegrations/Source/MoverMassIntegration/Public/MoverMassTranslators.h
@@ -1,9 +1,10 @@
-﻿// Copyright Epic Games, Inc. All Rights Reserved.
+// Copyright Epic Games, Inc. All Rights Reserved.
 
 #pragma once
 
 #include "MassCommonFragments.h"
 #include "MassTranslator.h"
+#include "MassEntityQuery.h"
 #include "MassObserverProcessor.h"
 #include "MoverMassTranslators.generated.h"
 
diff --git a/Engine/Plugins/Tests/RuntimeTests/Source/RuntimeTests/Public/EngineRuntimeTests.h b/Engine/Plugins/Tests/RuntimeTests/Source/RuntimeTests/Public/EngineRuntimeTests.h
--- a/Engine/Plugins/Tests/RuntimeTests/Source/RuntimeTests/Public/EngineRuntimeTests.h
+++ b/Engine/Plugins/Tests/RuntimeTests/Source/RuntimeTests/Public/EngineRuntimeTests.h
@@ -6,7 +6,7 @@
 #include "Tests/AutomationCommon.h"
 #include "MassProcessor.h"
 #include "MassEntityTypes.h"
-
+#include "MassEntityQuery.h"
 #include "EngineRuntimeTests.generated.h"
 
 #define UE_API RUNTIMETESTS_API
